NeuralNetworkInference plugin:
Disconnect Operators

diff --git a/Engine/Plugins/Experimental/NNI/Source/Private/NeuralNetworkInferenceQA/Private/UnitTester.cpp b/Engine/Plugins/Experimental/NNI/Source/Private/NeuralNetworkInferenceQA/Private/UnitTester.cpp
--- a/Engine/Plugins/Experimental/NNI/Source/Private/NeuralNetworkInferenceQA/Private/UnitTester.cpp
+++ b/Engine/Plugins/Experimental/NNI/Source/Private/NeuralNetworkInferenceQA/Private/UnitTester.cpp
@@ -26,8 +26,8 @@ bool FUnitTester::GlobalTest(const FString& InProjectContentDir, const FString&
 	UE_LOG(LogNeuralNetworkInferenceQA, Display, TEXT("---------------------------------------------------------------------------------------------------------------------------------"));
 	UE_LOG(LogNeuralNetworkInferenceQA, Display, TEXT("---------------------------------------------------------------------------------------------------------------------------------"));
 	UE_LOG(LogNeuralNetworkInferenceQA, Display, TEXT("---------------------------------------------------------------------------------------------------------------------------------"));
-	UE_LOG(LogNeuralNetworkInferenceQA, Display, TEXT("---------- 2. Operator Unit Testing"));
-	bDidGlobalTestPassed &= FOperatorUnitTester::GlobalTest(InProjectContentDir, InUnitTestRelativeDirectory);
+	//UE_LOG(LogNeuralNetworkInferenceQA, Display, TEXT("---------- 2. Operator Unit Testing"));
+	//bDidGlobalTestPassed &= FOperatorUnitTester::GlobalTest(InProjectContentDir, InUnitTestRelativeDirectory);
 
 	if (bDidGlobalTestPassed)
 	{
