Support NNI CPU inference when using an inference context.

diff --git a/Engine/Plugins/Experimental/NNI/Source/NeuralNetworkInference/Private/NeuralNetwork.cpp b/Engine/Plugins/Experimental/NNI/Source/NeuralNetworkInference/Private/NeuralNetwork.cpp
--- a/Engine/Plugins/Experimental/NNI/Source/NeuralNetworkInference/Private/NeuralNetwork.cpp
+++ b/Engine/Plugins/Experimental/NNI/Source/NeuralNetworkInference/Private/NeuralNetwork.cpp
@@ -735,13 +735,13 @@ int32 UNeuralNetwork::CreateInferenceContext()
 		UE_LOG(LogNeuralNetworkInference, Warning, TEXT("UNeuralNetwork::CreateInferenceContext(): Execution contexts currently only support synchronous Onyx models."));
 		return -1;
 	}
-	if (DeviceType != ENeuralDeviceType::GPU || InputDeviceType != ENeuralDeviceType::CPU || OutputDeviceType != ENeuralDeviceType::GPU)
+	if (InputDeviceType != ENeuralDeviceType::CPU)
 	{
-		UE_LOG(LogNeuralNetworkInference, Warning, TEXT("UNeuralNetwork::CreateInferenceContext(): Execution contexts currently only support GPU device with CPU input -> GPU ouput."));
+		UE_LOG(LogNeuralNetworkInference, Warning, TEXT("UNeuralNetwork::CreateInferenceContext(): Execution contexts currently only support CPU input."));
 		return -1;
 	}
 
-	return ImplBackEndUEAndORT->CreateInferenceContext();
+	return ImplBackEndUEAndORT->CreateInferenceContext(InputDeviceType, OutputDeviceType);
 }
 
 void UNeuralNetwork::DestroyInferenceContext(int32 ContextHandle)
@@ -749,6 +749,18 @@ void UNeuralNetwork::DestroyInferenceContext(int32 ContextHandle)
 	return ImplBackEndUEAndORT->DestroyInferenceContext(ContextHandle);
 }
 
+void UNeuralNetwork::Run(int32 ContextHandle)
+{
+	DECLARE_SCOPE_CYCLE_COUNTER(TEXT("UNeuralNetwork_Run"), STAT_UNeuralNetwork_Run, STATGROUP_MachineLearning);
+
+	FNeuralTimer RunTimer;
+	RunTimer.Tic();
+
+	ImplBackEndUEAndORT->Run(ContextHandle);
+
+	RunStatisticsEstimator.StoreSample(RunTimer.Toc());
+}
+
 void UNeuralNetwork::Run(FRDGBuilder& GraphBuilder, int32 ContextHandle)
 {
 	DECLARE_SCOPE_CYCLE_COUNTER(TEXT("UNeuralNetwork_Run"), STAT_UNeuralNetwork_Run, STATGROUP_MachineLearning);
diff --git a/Engine/Plugins/Experimental/NNI/Source/NeuralNetworkInference/Private/NeuralNetworkImplBackEndUEAndORT.cpp b/Engine/Plugins/Experimental/NNI/Source/NeuralNetworkInference/Private/NeuralNetworkImplBackEndUEAndORT.cpp
--- a/Engine/Plugins/Experimental/NNI/Source/NeuralNetworkInference/Private/NeuralNetworkImplBackEndUEAndORT.cpp
+++ b/Engine/Plugins/Experimental/NNI/Source/NeuralNetworkInference/Private/NeuralNetworkImplBackEndUEAndORT.cpp
@@ -1064,11 +1064,14 @@ void UNeuralNetwork::FImplBackEndUEAndORT::ClearResources()
 #endif //PLATFORM_WIN64
 }
 
-bool UNeuralNetwork::FImplBackEndUEAndORT::FInferenceContext::Init(Ort::Session* InSession, Ort::AllocatorWithDefaultOptions* InAllocator, Ort::MemoryInfo* InAllocatorInfo)
+bool UNeuralNetwork::FImplBackEndUEAndORT::FInferenceContext::Init(
+	Ort::Session* InSession, 
+	Ort::AllocatorWithDefaultOptions* InAllocator, 
+	Ort::MemoryInfo* InAllocatorInfo,
+	ENeuralDeviceType InInputDeviceType, 
+	ENeuralDeviceType InOutputDeviceType)
 {
 	// Current assumptions here (all these things can be fixed with more work): 
-	// Input tensors are CPU
-	// Output tensors are GPU 
 	// Float data types only
 	// No variable dimensions
 
@@ -1117,20 +1120,32 @@ bool UNeuralNetwork::FImplBackEndUEAndORT::FInferenceContext::Init(Ort::Session*
 			if (bInputTensors)
 			{
 				InputTensorNames.Emplace(TensorName);
-
 				InputTensors.Emplace(FNeuralTensor(TensorDataType, TensorSizes, ANSI_TO_TCHAR(TensorName), TensorType));
-
 				InputOrtTensors.Emplace(Ort::Value(nullptr));
-				LinkTensorToONNXRuntime(InputTensors[TensorIndex], InputOrtTensors[TensorIndex], *InAllocatorInfo);
+				
+				if (InInputDeviceType == ENeuralDeviceType::CPU)
+				{
+					LinkTensorToONNXRuntime(InputTensors[TensorIndex], InputOrtTensors[TensorIndex], *InAllocatorInfo);
+				}
+				else
+				{
+					InputTensors[TensorIndex].SetEnableGPU(true);
+				}
 			}
 			else
 			{
 				OutputTensorNames.Emplace(TensorName);
-
-				FNeuralTensor& Tensor = OutputTensors.Emplace_GetRef(FNeuralTensor(TensorDataType, TensorSizes, ANSI_TO_TCHAR(TensorName), TensorType));
-				Tensor.SetEnableGPU(true);
-
+				OutputTensors.Emplace(FNeuralTensor(TensorDataType, TensorSizes, ANSI_TO_TCHAR(TensorName), TensorType));
 				OutputOrtTensors.Emplace(Ort::Value(nullptr));
+
+				if (InOutputDeviceType == ENeuralDeviceType::CPU)
+				{
+					LinkTensorToONNXRuntime(OutputTensors[TensorIndex], OutputOrtTensors[TensorIndex], *InAllocatorInfo);
+				}
+				else
+				{
+					OutputTensors[TensorIndex].SetEnableGPU(true);
+				}
 			}
 
 			TypeInfo.release();
@@ -1185,11 +1200,11 @@ void UNeuralNetwork::FImplBackEndUEAndORT::FInferenceContext::ReleaseDMLAllocati
 
 #endif //WITH_UE_AND_ORT_SUPPORT
 
-int32 UNeuralNetwork::FImplBackEndUEAndORT::CreateInferenceContext()
+int32 UNeuralNetwork::FImplBackEndUEAndORT::CreateInferenceContext(ENeuralDeviceType InInputDeviceType, ENeuralDeviceType InOutputDeviceType)
 {
 #ifdef WITH_UE_AND_ORT_SUPPORT
 	const int32 Handle = Contexts.Emplace();
-	Contexts[Handle].Init(Session.Get(), Allocator.Get(), AllocatorInfo.Get());
+	Contexts[Handle].Init(Session.Get(), Allocator.Get(), AllocatorInfo.Get(), InInputDeviceType, InOutputDeviceType);
 	return Handle;
 #else
 	UE_LOG(LogNeuralNetworkInference, Warning, TEXT("FImplBackEndUEAndORT::CreateInferenceContext(): Platform or Operating System not suported yet for UEAndORT"
@@ -1210,6 +1225,18 @@ void UNeuralNetwork::FImplBackEndUEAndORT::DestroyInferenceContext(int32 InConte
 #endif
 }
 
+void UNeuralNetwork::FImplBackEndUEAndORT::Run(int32 InContextHandle)
+{
+#ifdef WITH_UE_AND_ORT_SUPPORT
+	FInferenceContext& Context = Contexts[InContextHandle];
+
+	Session.Get()->Run(
+		Ort::RunOptions{ nullptr },
+		Context.InputTensorNames.GetData(), &Context.InputOrtTensors[0], Context.InputTensorNames.Num(),
+		Context.OutputTensorNames.GetData(), &Context.OutputOrtTensors[0], Context.OutputTensorNames.Num());
+#endif
+}
+
 void UNeuralNetwork::FImplBackEndUEAndORT::Run(FRDGBuilder& GraphBuilder, int32 InContextHandle)
 {
 #ifdef WITH_UE_AND_ORT_SUPPORT
diff --git a/Engine/Plugins/Experimental/NNI/Source/NeuralNetworkInference/Private/NeuralNetworkImplBackEndUEAndORT.h b/Engine/Plugins/Experimental/NNI/Source/NeuralNetworkInference/Private/NeuralNetworkImplBackEndUEAndORT.h
--- a/Engine/Plugins/Experimental/NNI/Source/NeuralNetworkInference/Private/NeuralNetworkImplBackEndUEAndORT.h
+++ b/Engine/Plugins/Experimental/NNI/Source/NeuralNetworkInference/Private/NeuralNetworkImplBackEndUEAndORT.h
@@ -53,9 +53,10 @@ public:
 		const TArray<uint8>& InModelReadFromFileInBytes, const FString& InModelFullFilePath, const ENeuralDeviceType InDeviceType, const ENeuralDeviceType InInputDeviceType,
 		const ENeuralDeviceType InOutputDeviceType);
 
-	int32 CreateInferenceContext();
+	int32 CreateInferenceContext(ENeuralDeviceType InInputDeviceType, ENeuralDeviceType InOutputDeviceType);
 	void DestroyInferenceContext(int32 ContextHandle);
 
+	void Run(int32 ContextHandle);
 	void Run(FRDGBuilder& GraphBuilder, int32 ContextHandle);
 	void Run(const ENeuralSynchronousMode InSynchronousMode, const ENeuralDeviceType InDeviceType, const ENeuralDeviceType InInputDeviceType);
 
@@ -113,7 +114,7 @@ private:
 		TArray<Ort::Value> OutputOrtTensors;
 		TArray<void*> OutputDmlAllocation;
 
-		bool Init(Ort::Session* Session, Ort::AllocatorWithDefaultOptions* Allocator, Ort::MemoryInfo* AllocatorInfo);
+		bool Init(Ort::Session* Session, Ort::AllocatorWithDefaultOptions* Allocator, Ort::MemoryInfo* AllocatorInfo, ENeuralDeviceType InInputDeviceType, ENeuralDeviceType InOutputDeviceType);
 		void UpdateGPUAllocations(FRDGBuilder& GraphBuilder);
 		void Release();
 #ifdef PLATFORM_WIN64
diff --git a/Engine/Plugins/Experimental/NNI/Source/NeuralNetworkInference/Public/NeuralNetwork.h b/Engine/Plugins/Experimental/NNI/Source/NeuralNetworkInference/Public/NeuralNetwork.h
--- a/Engine/Plugins/Experimental/NNI/Source/NeuralNetworkInference/Public/NeuralNetwork.h
+++ b/Engine/Plugins/Experimental/NNI/Source/NeuralNetworkInference/Public/NeuralNetwork.h
@@ -240,7 +240,13 @@ public:
 
 	/**
 	 * Run() executes the forward pass of the current UNeuralNetwork given an inference context created with CreateInferenceContext().
-	 * This should be called on the render thread.
+	 * This version assumes CPU inference.
+	 */
+	void Run(int32 ContextHandle);
+
+	/**
+	 * Run() executes the forward pass of the current UNeuralNetwork given an inference context created with CreateInferenceContext().
+	 * This version assumes GPU inference and should be called on the render thread.
 	 */
 	void Run(class FRDGBuilder& GraphBuilder, int32 ContextHandle);
 
