Win64 Clang compile fix.

diff --git a/Engine/Plugins/Experimental/NNI/Source/Private/ThirdPartyHelperAndDLLLoader/Internal/ThirdPartyWarningDisabler.h b/Engine/Plugins/Experimental/NNI/Source/Private/ThirdPartyHelperAndDLLLoader/Internal/ThirdPartyWarningDisabler.h
--- a/Engine/Plugins/Experimental/NNI/Source/Private/ThirdPartyHelperAndDLLLoader/Internal/ThirdPartyWarningDisabler.h
+++ b/Engine/Plugins/Experimental/NNI/Source/Private/ThirdPartyHelperAndDLLLoader/Internal/ThirdPartyWarningDisabler.h
@@ -2,7 +2,7 @@
 
 #pragma once
 
-#include "CoreTypes.h"
+#include "CoreMinimal.h"
 
 /*
 Many third party headers require some care when importing. NNI third party includes should be wrapped like this:
