In case when both NNI and NNE plugins are enabled we should only use one DirectML.dll from NNE plugin

#rnx

diff --git a/Engine/Plugins/Experimental/NNI/Source/Private/ThirdPartyHelperAndDLLLoader/Private/Module.cpp b/Engine/Plugins/Experimental/NNI/Source/Private/ThirdPartyHelperAndDLLLoader/Private/Module.cpp
--- a/Engine/Plugins/Experimental/NNI/Source/Private/ThirdPartyHelperAndDLLLoader/Private/Module.cpp
+++ b/Engine/Plugins/Experimental/NNI/Source/Private/ThirdPartyHelperAndDLLLoader/Private/Module.cpp
@@ -11,6 +11,12 @@
 void FThirdPartyHelperAndDLLLoaderModule::StartupModule()
 {
 #ifdef PLATFORM_WIN64
+	TSharedPtr<IPlugin> NNEPlugin = IPluginManager::Get().FindPlugin(TEXT("NNE"));
+	if (NNEPlugin.IsValid())
+	{
+		return;
+	}
+
 	const FString PluginDir = IPluginManager::Get().FindPlugin(TEXT("NeuralNetworkInference"))->GetBaseDir();
 	const FString DirectMLRuntimeBinPath = PluginDir / TEXT(PREPROCESSOR_TO_STRING(DIRECTML_PLATFORM_PATH));
 	const FString DirectMLDLLPath = DirectMLRuntimeBinPath / TEXT(PREPROCESSOR_TO_STRING(DIRECTML_DLL_NAME));
