Fix DirectML loading issue in NNI plugin

#rnx

diff --git a/Engine/Plugins/Experimental/NNI/Source/Private/ThirdPartyHelperAndDLLLoader/Private/Module.cpp b/Engine/Plugins/Experimental/NNI/Source/Private/ThirdPartyHelperAndDLLLoader/Private/Module.cpp
--- a/Engine/Plugins/Experimental/NNI/Source/Private/ThirdPartyHelperAndDLLLoader/Private/Module.cpp
+++ b/Engine/Plugins/Experimental/NNI/Source/Private/ThirdPartyHelperAndDLLLoader/Private/Module.cpp
@@ -12,7 +12,7 @@ void FThirdPartyHelperAndDLLLoaderModule::StartupModule()
 {
 #ifdef PLATFORM_WIN64
 	TSharedPtr<IPlugin> NNEPlugin = IPluginManager::Get().FindPlugin(TEXT("NNE"));
-	if (NNEPlugin.IsValid())
+	if (NNEPlugin.IsValid() && NNEPlugin->IsEnabled())
 	{
 		return;
 	}
