Fix for inccorect DirectML.dll path

#rnx

diff --git a/Engine/Plugins/Experimental/NNI/Source/Private/ThirdPartyHelperAndDLLLoader/Private/Module.cpp b/Engine/Plugins/Experimental/NNI/Source/Private/ThirdPartyHelperAndDLLLoader/Private/Module.cpp
--- a/Engine/Plugins/Experimental/NNI/Source/Private/ThirdPartyHelperAndDLLLoader/Private/Module.cpp
+++ b/Engine/Plugins/Experimental/NNI/Source/Private/ThirdPartyHelperAndDLLLoader/Private/Module.cpp
@@ -11,7 +11,7 @@
 void FThirdPartyHelperAndDLLLoaderModule::StartupModule()
 {
 #ifdef PLATFORM_WIN64
-#ifdef WITH_EDITOR
+#if WITH_EDITOR
 
 	TSharedPtr<IPlugin> NNEPlugin = IPluginManager::Get().FindPlugin(TEXT("NNE"));
 	if (NNEPlugin.IsValid() && NNEPlugin->IsEnabled())
