Fix cooked editor not being able to load due to dll path

diff --git a/Engine/Plugins/Experimental/NNI/Source/Private/ThirdPartyHelperAndDLLLoader/Private/Module.cpp b/Engine/Plugins/Experimental/NNI/Source/Private/ThirdPartyHelperAndDLLLoader/Private/Module.cpp
--- a/Engine/Plugins/Experimental/NNI/Source/Private/ThirdPartyHelperAndDLLLoader/Private/Module.cpp
+++ b/Engine/Plugins/Experimental/NNI/Source/Private/ThirdPartyHelperAndDLLLoader/Private/Module.cpp
@@ -11,7 +11,7 @@
 void FThirdPartyHelperAndDLLLoaderModule::StartupModule()
 {
 #ifdef PLATFORM_WIN64
-#if WITH_EDITOR
+#if WITH_EDITOR && !IS_MONOLITHIC
 
 	TSharedPtr<IPlugin> NNEPlugin = IPluginManager::Get().FindPlugin(TEXT("NNE"));
 	if (NNEPlugin.IsValid() && NNEPlugin->IsEnabled())
