NNI: Remove RedirectCoutAndCErrToUeLog.cpp due to crash when used in multiple threads

#rnx

diff --git a/Engine/Plugins/Experimental/NNI/Source/NeuralNetworkInference/Private/NeuralNetworkImplBackEndUEAndORT.cpp b/Engine/Plugins/Experimental/NNI/Source/NeuralNetworkInference/Private/NeuralNetworkImplBackEndUEAndORT.cpp
--- a/Engine/Plugins/Experimental/NNI/Source/NeuralNetworkInference/Private/NeuralNetworkImplBackEndUEAndORT.cpp
+++ b/Engine/Plugins/Experimental/NNI/Source/NeuralNetworkInference/Private/NeuralNetworkImplBackEndUEAndORT.cpp
@@ -5,7 +5,6 @@
 #include "NeuralEnumClasses.h"
 #include "NeuralNetworkInferenceUtils.h"
 #include "NeuralNetworkInferenceUtilsGPU.h"
-#include "RedirectCoutAndCerrToUeLog.h"
 #if WITH_EDITOR
 #include "Misc/MessageDialog.h"
 #endif //WITH_EDITOR
@@ -239,7 +238,6 @@ void UNeuralNetwork::FImplBackEndUEAndORT::FNeuralNetworkAsyncTask::SetRunSessio
 
 void UNeuralNetwork::FImplBackEndUEAndORT::FNeuralNetworkAsyncTask::DoWork()
 {
-	const FRedirectCoutAndCerrToUeLog RedirectCoutAndCerrToUeLog;
 	if (SyncMode == ENeuralSynchronousMode::Asynchronous)
 	{
 		BackEnd->RunSessionAsync(DeviceType, InputDeviceType);
@@ -312,8 +310,6 @@ bool UNeuralNetwork::FImplBackEndUEAndORT::Load(TSharedPtr<FImplBackEndUEAndORT>
 	try
 #endif //WITH_EDITOR
 	{
-		const FRedirectCoutAndCerrToUeLog RedirectCoutAndCerrToUeLog;
-
 		// Avoid multi-threaded crashes
 		if (InOutImplBackEndUEAndORT.IsValid())
 		{
diff --git a/Engine/Plugins/Experimental/NNI/Source/Private/ThirdPartyHelperAndDLLLoader/Internal/RedirectCoutAndCerrToUeLog.h b/Engine/Plugins/Experimental/NNI/Source/Private/ThirdPartyHelperAndDLLLoader/Internal/RedirectCoutAndCerrToUeLog.h
deleted file mode 100644
diff --git a/Engine/Plugins/Experimental/NNI/Source/Private/ThirdPartyHelperAndDLLLoader/Private/RedirectCoutAndCerrToUeLog.cpp b/Engine/Plugins/Experimental/NNI/Source/Private/ThirdPartyHelperAndDLLLoader/Private/RedirectCoutAndCerrToUeLog.cpp
deleted file mode 100644
