Fix spelling error
#rnx

diff --git a/Engine/Plugins/Experimental/NNI/Source/Private/ThirdPartyHelperAndDLLLoader/Private/Module.cpp b/Engine/Plugins/Experimental/NNI/Source/Private/ThirdPartyHelperAndDLLLoader/Private/Module.cpp
--- a/Engine/Plugins/Experimental/NNI/Source/Private/ThirdPartyHelperAndDLLLoader/Private/Module.cpp
+++ b/Engine/Plugins/Experimental/NNI/Source/Private/ThirdPartyHelperAndDLLLoader/Private/Module.cpp
@@ -17,7 +17,7 @@ void FThirdPartyHelperAndDLLLoaderModule::StartupModule()
 		UE_LOG(
 			LogNeuralNetworkInferenceThirdPartyHelperAndDLLLoader,
 			Display,
-			TEXT("Please disable NNE plugin, both NeuralNetorkInference and NNE plugins can't be used at a same time"));
+			TEXT("Please disable NNE plugin, both NeuralNetworkInference and NNE plugins can't be used at a same time"));
 	}
 
 	const FString DirectMLRuntimeBinPath = FString(FPlatformProcess::BaseDir()) / TEXT(PREPROCESSOR_TO_STRING(DIRECTML_PATH));
