Fixing NNI to only use WinPixEventRuntime when profiling is enabled in non-shipping and non-test builds.

diff --git a/Engine/Plugins/Experimental/NNI/Source/NeuralNetworkInference/NeuralNetworkInference.Build.cs b/Engine/Plugins/Experimental/NNI/Source/NeuralNetworkInference/NeuralNetworkInference.Build.cs
--- a/Engine/Plugins/Experimental/NNI/Source/NeuralNetworkInference/NeuralNetworkInference.Build.cs
+++ b/Engine/Plugins/Experimental/NNI/Source/NeuralNetworkInference/NeuralNetworkInference.Build.cs
@@ -60,7 +60,13 @@ public class NeuralNetworkInference : ModuleRules
 
 				AddEngineThirdPartyPrivateStaticDependencies(Target, "DX12");
 				AddEngineThirdPartyPrivateStaticDependencies(Target, "DirectML_1_8_0");
-				AddEngineThirdPartyPrivateStaticDependencies(Target, "WinPixEventRuntime");
+
+				if (Target.WindowsPlatform.bPixProfilingEnabled &&
+					Target.Configuration != UnrealTargetConfiguration.Shipping &&
+					Target.Configuration != UnrealTargetConfiguration.Test)
+				{
+					PrivateDependencyModuleNames.Add("WinPixEventRuntime");
+				}
 			}
 		}
 
